Update Rsync filter list for new version of Mono.

#rnx

diff --git a/Engine/Build/Rsync/RsyncEngine.txt b/Engine/Build/Rsync/RsyncEngine.txt
--- a/Engine/Build/Rsync/RsyncEngine.txt
+++ b/Engine/Build/Rsync/RsyncEngine.txt
@@ -59,154 +59,145 @@
 + /Binaries/ThirdParty/Mono/Mac/
 + /Binaries/ThirdParty/Mono/Mac/**
 - /Binaries/ThirdParty/Mono/**
+protect /Binaries/ThirdParty/Mono/Mac/FixMonoFiles.version
 protect /Binaries/ThirdParty/Mono/Mac/bin/mono
-protect /Binaries/ThirdParty/Mono/Mac/lib/mono/4.5/Accessibility.dll
-protect /Binaries/ThirdParty/Mono/Mac/lib/mono/4.5/Commons.Xml.Relaxng.dll
-protect /Binaries/ThirdParty/Mono/Mac/lib/mono/4.5/cscompmgd.dll
-protect /Binaries/ThirdParty/Mono/Mac/lib/mono/4.5/CustomMarshalers.dll
-protect /Binaries/ThirdParty/Mono/Mac/lib/mono/4.5/I18N.CJK.dll
-protect /Binaries/ThirdParty/Mono/Mac/lib/mono/4.5/I18N.dll
-protect /Binaries/ThirdParty/Mono/Mac/lib/mono/4.5/I18N.MidEast.dll
-protect /Binaries/ThirdParty/Mono/Mac/lib/mono/4.5/I18N.Other.dll
-protect /Binaries/ThirdParty/Mono/Mac/lib/mono/4.5/I18N.Rare.dll
-protect /Binaries/ThirdParty/Mono/Mac/lib/mono/4.5/I18N.West.dll
-protect /Binaries/ThirdParty/Mono/Mac/lib/mono/4.5/IBM.Data.DB2.dll
-protect /Binaries/ThirdParty/Mono/Mac/lib/mono/4.5/ICSharpCode.SharpZipLib.dll
-protect /Binaries/ThirdParty/Mono/Mac/lib/mono/4.5/Microsoft.Build.dll
-protect /Binaries/ThirdParty/Mono/Mac/lib/mono/4.5/Microsoft.Build.Engine.dll
-protect /Binaries/ThirdParty/Mono/Mac/lib/mono/4.5/Microsoft.Build.Framework.dll
-protect /Binaries/ThirdParty/Mono/Mac/lib/mono/4.5/Microsoft.Build.Tasks.v4.0.dll
-protect /Binaries/ThirdParty/Mono/Mac/lib/mono/4.5/Microsoft.Build.Utilities.v4.0.dll
+protect /Binaries/ThirdParty/Mono/Mac/lib/mono/xbuild/14.0/bin/Microsoft.Build.Engine.dll
+protect /Binaries/ThirdParty/Mono/Mac/lib/mono/xbuild/14.0/bin/Mono.XBuild.Tasks.dll
+protect /Binaries/ThirdParty/Mono/Mac/lib/mono/xbuild/14.0/bin/Microsoft.Build.dll
+protect /Binaries/ThirdParty/Mono/Mac/lib/mono/xbuild/14.0/bin/Microsoft.Build.Utilities.Core.dll
+protect /Binaries/ThirdParty/Mono/Mac/lib/mono/xbuild/14.0/bin/Microsoft.Build.Tasks.Core.dll
+protect /Binaries/ThirdParty/Mono/Mac/lib/mono/xbuild/14.0/bin/Microsoft.Build.Framework.dll
+protect /Binaries/ThirdParty/Mono/Mac/lib/mono/xbuild/12.0/bin/Microsoft.Build.Engine.dll
+protect /Binaries/ThirdParty/Mono/Mac/lib/mono/xbuild/12.0/bin/Mono.XBuild.Tasks.dll
+protect /Binaries/ThirdParty/Mono/Mac/lib/mono/xbuild/12.0/bin/Microsoft.Build.dll
+protect /Binaries/ThirdParty/Mono/Mac/lib/mono/xbuild/12.0/bin/Microsoft.Build.Framework.dll
+protect /Binaries/ThirdParty/Mono/Mac/lib/mono/xbuild/12.0/bin/Microsoft.Build.Utilities.v12.0.dll
+protect /Binaries/ThirdParty/Mono/Mac/lib/mono/xbuild/12.0/bin/Microsoft.Build.Tasks.v12.0.dll
 protect /Binaries/ThirdParty/Mono/Mac/lib/mono/4.5/Microsoft.CSharp.dll
-protect /Binaries/ThirdParty/Mono/Mac/lib/mono/4.5/Microsoft.VisualC.dll
-protect /Binaries/ThirdParty/Mono/Mac/lib/mono/4.5/Microsoft.Web.Infrastructure.dll
-protect /Binaries/ThirdParty/Mono/Mac/lib/mono/4.5/Mono.C5.dll
-protect /Binaries/ThirdParty/Mono/Mac/lib/mono/4.5/Mono.Cairo.dll
-protect /Binaries/ThirdParty/Mono/Mac/lib/mono/4.5/Mono.CodeContracts.dll
-protect /Binaries/ThirdParty/Mono/Mac/lib/mono/4.5/Mono.CompilerServices.SymbolWriter.dll
-protect /Binaries/ThirdParty/Mono/Mac/lib/mono/4.5/Mono.CSharp.dll
-protect /Binaries/ThirdParty/Mono/Mac/lib/mono/4.5/Mono.Data.Sqlite.dll
-protect /Binaries/ThirdParty/Mono/Mac/lib/mono/4.5/Mono.Data.Tds.dll
-protect /Binaries/ThirdParty/Mono/Mac/lib/mono/4.5/Mono.Debugger.Soft.dll
-protect /Binaries/ThirdParty/Mono/Mac/lib/mono/4.5/Mono.Http.dll
-protect /Binaries/ThirdParty/Mono/Mac/lib/mono/4.5/Mono.Management.dll
-protect /Binaries/ThirdParty/Mono/Mac/lib/mono/4.5/Mono.Messaging.dll
-protect /Binaries/ThirdParty/Mono/Mac/lib/mono/4.5/Mono.Messaging.RabbitMQ.dll
-protect /Binaries/ThirdParty/Mono/Mac/lib/mono/4.5/Mono.Parallel.dll
+protect /Binaries/ThirdParty/Mono/Mac/lib/mono/4.5/System.Web.dll
+protect /Binaries/ThirdParty/Mono/Mac/lib/mono/4.5/System.ServiceModel.Web.dll
+protect /Binaries/ThirdParty/Mono/Mac/lib/mono/4.5/mod-mono-server4.exe
+protect /Binaries/ThirdParty/Mono/Mac/lib/mono/4.5/System.Xml.Linq.dll
 protect /Binaries/ThirdParty/Mono/Mac/lib/mono/4.5/Mono.Posix.dll
-protect /Binaries/ThirdParty/Mono/Mac/lib/mono/4.5/Mono.Security.dll
-protect /Binaries/ThirdParty/Mono/Mac/lib/mono/4.5/Mono.Security.Win32.dll
-protect /Binaries/ThirdParty/Mono/Mac/lib/mono/4.5/Mono.Simd.dll
-protect /Binaries/ThirdParty/Mono/Mac/lib/mono/4.5/Mono.Tasklets.dll
-protect /Binaries/ThirdParty/Mono/Mac/lib/mono/4.5/Mono.WebBrowser.dll
-protect /Binaries/ThirdParty/Mono/Mac/lib/mono/4.5/Mono.XBuild.Tasks.dll
-protect /Binaries/ThirdParty/Mono/Mac/lib/mono/4.5/Novell.Directory.Ldap.dll
-protect /Binaries/ThirdParty/Mono/Mac/lib/mono/4.5/nunit-console-runner.dll
-protect /Binaries/ThirdParty/Mono/Mac/lib/mono/4.5/nunit.core.dll
-protect /Binaries/ThirdParty/Mono/Mac/lib/mono/4.5/nunit.core.extensions.dll
-protect /Binaries/ThirdParty/Mono/Mac/lib/mono/4.5/nunit.core.interfaces.dll
-protect /Binaries/ThirdParty/Mono/Mac/lib/mono/4.5/nunit.framework.dll
-protect /Binaries/ThirdParty/Mono/Mac/lib/mono/4.5/nunit.framework.extensions.dll
-protect /Binaries/ThirdParty/Mono/Mac/lib/mono/4.5/nunit.mocks.dll
-protect /Binaries/ThirdParty/Mono/Mac/lib/mono/4.5/nunit.util.dll
-protect /Binaries/ThirdParty/Mono/Mac/lib/mono/4.5/PEAPI.dll
-protect /Binaries/ThirdParty/Mono/Mac/lib/mono/4.5/RabbitMQ.Client.dll
-protect /Binaries/ThirdParty/Mono/Mac/lib/mono/4.5/System.ComponentModel.Composition.dll
-protect /Binaries/ThirdParty/Mono/Mac/lib/mono/4.5/System.ComponentModel.DataAnnotations.dll
-protect /Binaries/ThirdParty/Mono/Mac/lib/mono/4.5/System.Configuration.dll
+protect /Binaries/ThirdParty/Mono/Mac/lib/mono/4.5/System.Web.DynamicData.dll
+protect /Binaries/ThirdParty/Mono/Mac/lib/mono/4.5/System.Web.RegularExpressions.dll
+protect /Binaries/ThirdParty/Mono/Mac/lib/mono/4.5/System.Dynamic.dll
+protect /Binaries/ThirdParty/Mono/Mac/lib/mono/4.5/System.Net.dll
+protect /Binaries/ThirdParty/Mono/Mac/lib/mono/4.5/System.IO.Compression.FileSystem.dll
 protect /Binaries/ThirdParty/Mono/Mac/lib/mono/4.5/System.Configuration.Install.dll
-protect /Binaries/ThirdParty/Mono/Mac/lib/mono/4.5/System.Core.dll
-protect /Binaries/ThirdParty/Mono/Mac/lib/mono/4.5/System.Data.DataSetExtensions.dll
-protect /Binaries/ThirdParty/Mono/Mac/lib/mono/4.5/System.Data.dll
-protect /Binaries/ThirdParty/Mono/Mac/lib/mono/4.5/System.Data.Entity.dll
+protect /Binaries/ThirdParty/Mono/Mac/lib/mono/4.5/System.Management.dll
+protect /Binaries/ThirdParty/Mono/Mac/lib/mono/4.5/System.Web.Razor.dll
+protect /Binaries/ThirdParty/Mono/Mac/lib/mono/4.5/Microsoft.Build.Engine.dll
 protect /Binaries/ThirdParty/Mono/Mac/lib/mono/4.5/System.Data.Linq.dll
+protect /Binaries/ThirdParty/Mono/Mac/lib/mono/4.5/Mono.Btls.Interface.dll
+protect /Binaries/ThirdParty/Mono/Mac/lib/mono/4.5/System.Web.Mobile.dll
+protect /Binaries/ThirdParty/Mono/Mac/lib/mono/4.5/System.Runtime.DurableInstancing.dll
+protect /Binaries/ThirdParty/Mono/Mac/lib/mono/4.5/Microsoft.VisualBasic.dll
 protect /Binaries/ThirdParty/Mono/Mac/lib/mono/4.5/System.Data.OracleClient.dll
-protect /Binaries/ThirdParty/Mono/Mac/lib/mono/4.5/System.Data.Services.Client.dll
-protect /Binaries/ThirdParty/Mono/Mac/lib/mono/4.5/System.Data.Services.dll
+protect /Binaries/ThirdParty/Mono/Mac/lib/mono/4.5/System.Configuration.dll
+protect /Binaries/ThirdParty/Mono/Mac/lib/mono/4.5/System.Json.dll
 protect /Binaries/ThirdParty/Mono/Mac/lib/mono/4.5/System.Design.dll
-protect /Binaries/ThirdParty/Mono/Mac/lib/mono/4.5/System.DirectoryServices.dll
-protect /Binaries/ThirdParty/Mono/Mac/lib/mono/4.5/System.DirectoryServices.Protocols.dll
-protect /Binaries/ThirdParty/Mono/Mac/lib/mono/4.5/System.dll
+protect /Binaries/ThirdParty/Mono/Mac/lib/mono/4.5/System.Net.Http.dll
+protect /Binaries/ThirdParty/Mono/Mac/lib/mono/4.5/System.Workflow.Activities.dll
+protect /Binaries/ThirdParty/Mono/Mac/lib/mono/4.5/Mono.WebBrowser.dll
+protect /Binaries/ThirdParty/Mono/Mac/lib/mono/4.5/System.ServiceModel.dll
+protect /Binaries/ThirdParty/Mono/Mac/lib/mono/4.5/System.Xaml.dll
+protect /Binaries/ThirdParty/Mono/Mac/lib/mono/4.5/System.Messaging.dll
 protect /Binaries/ThirdParty/Mono/Mac/lib/mono/4.5/System.Drawing.Design.dll
-protect /Binaries/ThirdParty/Mono/Mac/lib/mono/4.5/System.Drawing.dll
-protect /Binaries/ThirdParty/Mono/Mac/lib/mono/4.5/System.Dynamic.dll
-protect /Binaries/ThirdParty/Mono/Mac/lib/mono/4.5/System.EnterpriseServices.dll
-protect /Binaries/ThirdParty/Mono/Mac/lib/mono/4.5/System.IdentityModel.dll
-protect /Binaries/ThirdParty/Mono/Mac/lib/mono/4.5/System.IdentityModel.Selectors.dll
-protect /Binaries/ThirdParty/Mono/Mac/lib/mono/4.5/System.IO.Compression.dll
-protect /Binaries/ThirdParty/Mono/Mac/lib/mono/4.5/System.IO.Compression.FileSystem.dll
-protect /Binaries/ThirdParty/Mono/Mac/lib/mono/4.5/System.Json.dll
+protect /Binaries/ThirdParty/Mono/Mac/lib/mono/4.5/System.Web.WebPages.dll
+protect /Binaries/ThirdParty/Mono/Mac/lib/mono/4.5/Mono.Messaging.dll
+protect /Binaries/ThirdParty/Mono/Mac/lib/mono/4.5/Mono.Security.Win32.dll
+protect /Binaries/ThirdParty/Mono/Mac/lib/mono/4.5/Mono.XBuild.Tasks.dll
+protect /Binaries/ThirdParty/Mono/Mac/lib/mono/4.5/Mono.Data.Tds.dll
+protect /Binaries/ThirdParty/Mono/Mac/lib/mono/4.5/Mono.CompilerServices.SymbolWriter.dll
+protect /Binaries/ThirdParty/Mono/Mac/lib/mono/4.5/Accessibility.dll
 protect /Binaries/ThirdParty/Mono/Mac/lib/mono/4.5/System.Json.Microsoft.dll
-protect /Binaries/ThirdParty/Mono/Mac/lib/mono/4.5/System.Management.dll
-protect /Binaries/ThirdParty/Mono/Mac/lib/mono/4.5/System.Messaging.dll
-protect /Binaries/ThirdParty/Mono/Mac/lib/mono/4.5/System.Net.dll
-protect /Binaries/ThirdParty/Mono/Mac/lib/mono/4.5/System.Net.Http.dll
-protect /Binaries/ThirdParty/Mono/Mac/lib/mono/4.5/System.Net.Http.Formatting.dll
-protect /Binaries/ThirdParty/Mono/Mac/lib/mono/4.5/System.Net.Http.WebRequest.dll
-protect /Binaries/ThirdParty/Mono/Mac/lib/mono/4.5/System.Numerics.dll
-protect /Binaries/ThirdParty/Mono/Mac/lib/mono/4.5/System.Reactive.Core.dll
-protect /Binaries/ThirdParty/Mono/Mac/lib/mono/4.5/System.Reactive.Debugger.dll
-protect /Binaries/ThirdParty/Mono/Mac/lib/mono/4.5/System.Reactive.Experimental.dll
-protect /Binaries/ThirdParty/Mono/Mac/lib/mono/4.5/System.Reactive.Interfaces.dll
-protect /Binaries/ThirdParty/Mono/Mac/lib/mono/4.5/System.Reactive.Linq.dll
-protect /Binaries/ThirdParty/Mono/Mac/lib/mono/4.5/System.Reactive.Observable.Aliases.dll
-protect /Binaries/ThirdParty/Mono/Mac/lib/mono/4.5/System.Reactive.PlatformServices.dll
-protect /Binaries/ThirdParty/Mono/Mac/lib/mono/4.5/System.Reactive.Providers.dll
-protect /Binaries/ThirdParty/Mono/Mac/lib/mono/4.5/System.Reactive.Runtime.Remoting.dll
-protect /Binaries/ThirdParty/Mono/Mac/lib/mono/4.5/System.Reactive.Windows.Forms.dll
-protect /Binaries/ThirdParty/Mono/Mac/lib/mono/4.5/System.Reactive.Windows.Threading.dll
-protect /Binaries/ThirdParty/Mono/Mac/lib/mono/4.5/System.Runtime.Caching.dll
-protect /Binaries/ThirdParty/Mono/Mac/lib/mono/4.5/System.Runtime.DurableInstancing.dll
-protect /Binaries/ThirdParty/Mono/Mac/lib/mono/4.5/System.Runtime.Remoting.dll
+protect /Binaries/ThirdParty/Mono/Mac/lib/mono/4.5/System.Web.WebPages.Razor.dll
+protect /Binaries/ThirdParty/Mono/Mac/lib/mono/4.5/System.Web.Http.WebHost.dll
 protect /Binaries/ThirdParty/Mono/Mac/lib/mono/4.5/System.Runtime.Serialization.dll
-protect /Binaries/ThirdParty/Mono/Mac/lib/mono/4.5/System.Runtime.Serialization.Formatters.Soap.dll
-protect /Binaries/ThirdParty/Mono/Mac/lib/mono/4.5/System.Security.dll
-protect /Binaries/ThirdParty/Mono/Mac/lib/mono/4.5/System.ServiceModel.Activation.dll
-protect /Binaries/ThirdParty/Mono/Mac/lib/mono/4.5/System.ServiceModel.Discovery.dll
-protect /Binaries/ThirdParty/Mono/Mac/lib/mono/4.5/System.ServiceModel.dll
-protect /Binaries/ThirdParty/Mono/Mac/lib/mono/4.5/System.ServiceModel.Routing.dll
-protect /Binaries/ThirdParty/Mono/Mac/lib/mono/4.5/System.ServiceModel.Web.dll
-protect /Binaries/ThirdParty/Mono/Mac/lib/mono/4.5/System.ServiceProcess.dll
-protect /Binaries/ThirdParty/Mono/Mac/lib/mono/4.5/System.Threading.Tasks.Dataflow.dll
-protect /Binaries/ThirdParty/Mono/Mac/lib/mono/4.5/System.Transactions.dll
-protect /Binaries/ThirdParty/Mono/Mac/lib/mono/4.5/System.Web.Abstractions.dll
-protect /Binaries/ThirdParty/Mono/Mac/lib/mono/4.5/System.Web.ApplicationServices.dll
-protect /Binaries/ThirdParty/Mono/Mac/lib/mono/4.5/System.Web.dll
-protect /Binaries/ThirdParty/Mono/Mac/lib/mono/4.5/System.Web.DynamicData.dll
-protect /Binaries/ThirdParty/Mono/Mac/lib/mono/4.5/System.Web.Extensions.Design.dll
+protect /Binaries/ThirdParty/Mono/Mac/lib/mono/4.5/System.IO.Compression.dll
 protect /Binaries/ThirdParty/Mono/Mac/lib/mono/4.5/System.Web.Extensions.dll
-protect /Binaries/ThirdParty/Mono/Mac/lib/mono/4.5/System.Web.Http.dll
-protect /Binaries/ThirdParty/Mono/Mac/lib/mono/4.5/System.Web.Http.SelfHost.dll
-protect /Binaries/ThirdParty/Mono/Mac/lib/mono/4.5/System.Web.Http.WebHost.dll
-protect /Binaries/ThirdParty/Mono/Mac/lib/mono/4.5/System.Web.Mvc.dll
-protect /Binaries/ThirdParty/Mono/Mac/lib/mono/4.5/System.Web.Razor.dll
+protect /Binaries/ThirdParty/Mono/Mac/lib/mono/4.5/WindowsBase.dll
+protect /Binaries/ThirdParty/Mono/Mac/lib/mono/4.5/System.Net.Http.WebRequest.dll
+protect /Binaries/ThirdParty/Mono/Mac/lib/mono/4.5/System.Data.DataSetExtensions.dll
+protect /Binaries/ThirdParty/Mono/Mac/lib/mono/4.5/System.IdentityModel.dll
+protect /Binaries/ThirdParty/Mono/Mac/lib/mono/4.5/Microsoft.Build.Utilities.v4.0.dll
+protect /Binaries/ThirdParty/Mono/Mac/lib/mono/4.5/System.Web.Abstractions.dll
+protect /Binaries/ThirdParty/Mono/Mac/lib/mono/4.5/Mono.Debugger.Soft.dll
 protect /Binaries/ThirdParty/Mono/Mac/lib/mono/4.5/System.Web.Routing.dll
-protect /Binaries/ThirdParty/Mono/Mac/lib/mono/4.5/System.Web.Services.dll
-protect /Binaries/ThirdParty/Mono/Mac/lib/mono/4.5/System.Web.WebPages.Deployment.dll
-protect /Binaries/ThirdParty/Mono/Mac/lib/mono/4.5/System.Web.WebPages.dll
-protect /Binaries/ThirdParty/Mono/Mac/lib/mono/4.5/System.Web.WebPages.Razor.dll
+protect /Binaries/ThirdParty/Mono/Mac/lib/mono/4.5/Mono.Simd.dll
+protect /Binaries/ThirdParty/Mono/Mac/lib/mono/4.5/System.Xml.Serialization.dll
+protect /Binaries/ThirdParty/Mono/Mac/lib/mono/4.5/System.ServiceModel.Activation.dll
+protect /Binaries/ThirdParty/Mono/Mac/lib/mono/4.5/Microsoft.Build.Tasks.v4.0.dll
+protect /Binaries/ThirdParty/Mono/Mac/lib/mono/4.5/I18N.dll
+protect /Binaries/ThirdParty/Mono/Mac/lib/mono/4.5/Mono.Profiler.Log.dll
+protect /Binaries/ThirdParty/Mono/Mac/lib/mono/4.5/I18N.West.dll
+protect /Binaries/ThirdParty/Mono/Mac/lib/mono/4.5/SMDiagnostics.dll
+protect /Binaries/ThirdParty/Mono/Mac/lib/mono/4.5/System.DirectoryServices.dll
+protect /Binaries/ThirdParty/Mono/Mac/lib/mono/4.5/Microsoft.Build.dll
+protect /Binaries/ThirdParty/Mono/Mac/lib/mono/4.5/mono-fpm.exe
+protect /Binaries/ThirdParty/Mono/Mac/lib/mono/4.5/xsp4.exe
+protect /Binaries/ThirdParty/Mono/Mac/lib/mono/4.5/Commons.Xml.Relaxng.dll
 protect /Binaries/ThirdParty/Mono/Mac/lib/mono/4.5/System.Windows.dll
-protect /Binaries/ThirdParty/Mono/Mac/lib/mono/4.5/System.Windows.Forms.DataVisualization.dll
-protect /Binaries/ThirdParty/Mono/Mac/lib/mono/4.5/System.Windows.Forms.dll
-protect /Binaries/ThirdParty/Mono/Mac/lib/mono/4.5/System.Xaml.dll
+protect /Binaries/ThirdParty/Mono/Mac/lib/mono/4.5/System.Deployment.dll
+protect /Binaries/ThirdParty/Mono/Mac/lib/mono/4.5/Mono.Data.Sqlite.dll
+protect /Binaries/ThirdParty/Mono/Mac/lib/mono/4.5/Mono.Parallel.dll
+protect /Binaries/ThirdParty/Mono/Mac/lib/mono/4.5/System.DirectoryServices.Protocols.dll
+protect /Binaries/ThirdParty/Mono/Mac/lib/mono/4.5/System.Runtime.Serialization.Formatters.Soap.dll
+protect /Binaries/ThirdParty/Mono/Mac/lib/mono/4.5/System.Workflow.ComponentModel.dll
+protect /Binaries/ThirdParty/Mono/Mac/lib/mono/4.5/PEAPI.dll
+protect /Binaries/ThirdParty/Mono/Mac/lib/mono/4.5/System.Web.Http.SelfHost.dll
+protect /Binaries/ThirdParty/Mono/Mac/lib/mono/4.5/System.EnterpriseServices.dll
+protect /Binaries/ThirdParty/Mono/Mac/lib/mono/4.5/I18N.CJK.dll
+protect /Binaries/ThirdParty/Mono/Mac/lib/mono/4.5/System.Data.Services.Client.dll
+protect /Binaries/ThirdParty/Mono/Mac/lib/mono/4.5/Mono.CSharp.dll
+protect /Binaries/ThirdParty/Mono/Mac/lib/mono/4.5/Mono.Management.dll
+protect /Binaries/ThirdParty/Mono/Mac/lib/mono/4.5/System.Workflow.Runtime.dll
+protect /Binaries/ThirdParty/Mono/Mac/lib/mono/4.5/Mono.Tasklets.dll
+protect /Binaries/ThirdParty/Mono/Mac/lib/mono/4.5/cscompmgd.dll
+protect /Binaries/ThirdParty/Mono/Mac/lib/mono/4.5/System.ServiceModel.Internals.dll
+protect /Binaries/ThirdParty/Mono/Mac/lib/mono/4.5/System.Transactions.dll
+protect /Binaries/ThirdParty/Mono/Mac/lib/mono/4.5/System.Threading.Tasks.Dataflow.dll
+protect /Binaries/ThirdParty/Mono/Mac/lib/mono/4.5/System.Web.Services.dll
+protect /Binaries/ThirdParty/Mono/Mac/lib/mono/4.5/System.Data.Entity.dll
+protect /Binaries/ThirdParty/Mono/Mac/lib/mono/4.5/System.Runtime.Remoting.dll
+protect /Binaries/ThirdParty/Mono/Mac/lib/mono/4.5/Mono.Security.dll
+protect /Binaries/ThirdParty/Mono/Mac/lib/mono/4.5/I18N.MidEast.dll
+protect /Binaries/ThirdParty/Mono/Mac/lib/mono/4.5/System.Data.dll
+protect /Binaries/ThirdParty/Mono/Mac/lib/mono/4.5/System.dll
+protect /Binaries/ThirdParty/Mono/Mac/lib/mono/4.5/Mono.CodeContracts.dll
+protect /Binaries/ThirdParty/Mono/Mac/lib/mono/4.5/I18N.Rare.dll
+protect /Binaries/ThirdParty/Mono/Mac/lib/mono/4.5/System.Runtime.Caching.dll
+protect /Binaries/ThirdParty/Mono/Mac/lib/mono/4.5/System.Web.Http.dll
+protect /Binaries/ThirdParty/Mono/Mac/lib/mono/4.5/System.Web.Extensions.Design.dll
+protect /Binaries/ThirdParty/Mono/Mac/lib/mono/4.5/Microsoft.Web.Infrastructure.dll
+protect /Binaries/ThirdParty/Mono/Mac/lib/mono/4.5/I18N.Other.dll
+protect /Binaries/ThirdParty/Mono/Mac/lib/mono/4.5/Microsoft.Build.Framework.dll
+protect /Binaries/ThirdParty/Mono/Mac/lib/mono/4.5/System.Numerics.Vectors.dll
+protect /Binaries/ThirdParty/Mono/Mac/lib/mono/4.5/System.ServiceProcess.dll
+protect /Binaries/ThirdParty/Mono/Mac/lib/mono/4.5/CustomMarshalers.dll
+protect /Binaries/ThirdParty/Mono/Mac/lib/mono/4.5/Mono.Http.dll
+protect /Binaries/ThirdParty/Mono/Mac/lib/mono/4.5/System.Numerics.dll
 protect /Binaries/ThirdParty/Mono/Mac/lib/mono/4.5/System.Xml.dll
-protect /Binaries/ThirdParty/Mono/Mac/lib/mono/4.5/System.Xml.Linq.dll
-protect /Binaries/ThirdParty/Mono/Mac/lib/mono/4.5/System.Xml.Serialization.dll
+protect /Binaries/ThirdParty/Mono/Mac/lib/mono/4.5/System.IdentityModel.Selectors.dll
+protect /Binaries/ThirdParty/Mono/Mac/lib/mono/4.5/System.ComponentModel.Composition.dll
+protect /Binaries/ThirdParty/Mono/Mac/lib/mono/4.5/Mono.WebServer2.dll
+protect /Binaries/ThirdParty/Mono/Mac/lib/mono/4.5/System.Web.WebPages.Deployment.dll
+protect /Binaries/ThirdParty/Mono/Mac/lib/mono/4.5/System.Drawing.dll
+protect /Binaries/ThirdParty/Mono/Mac/lib/mono/4.5/System.Data.Services.dll
+protect /Binaries/ThirdParty/Mono/Mac/lib/mono/4.5/System.Reflection.Context.dll
+protect /Binaries/ThirdParty/Mono/Mac/lib/mono/4.5/System.Net.Http.Formatting.dll
+protect /Binaries/ThirdParty/Mono/Mac/lib/mono/4.5/System.ServiceModel.Routing.dll
+protect /Binaries/ThirdParty/Mono/Mac/lib/mono/4.5/System.Security.dll
 protect /Binaries/ThirdParty/Mono/Mac/lib/mono/4.5/WebMatrix.Data.dll
-protect /Binaries/ThirdParty/Mono/Mac/lib/mono/4.5/WindowsBase.dll
-protect /Binaries/ThirdParty/Mono/Mac/lib/mono/xbuild/12.0/bin/Microsoft.Build.dll
-protect /Binaries/ThirdParty/Mono/Mac/lib/mono/xbuild/12.0/bin/Microsoft.Build.Engine.dll
-protect /Binaries/ThirdParty/Mono/Mac/lib/mono/xbuild/12.0/bin/Microsoft.Build.Framework.dll
-protect /Binaries/ThirdParty/Mono/Mac/lib/mono/xbuild/12.0/bin/Microsoft.Build.Tasks.v12.0.dll
-protect /Binaries/ThirdParty/Mono/Mac/lib/mono/xbuild/12.0/bin/Microsoft.Build.Utilities.v12.0.dll
-protect /Binaries/ThirdParty/Mono/Mac/lib/mono/xbuild/12.0/bin/Mono.XBuild.Tasks.dll
-protect /Binaries/ThirdParty/Mono/Mac/lib/mono/xbuild/14.0/bin/Microsoft.Build.dll
-protect /Binaries/ThirdParty/Mono/Mac/lib/mono/xbuild/14.0/bin/Microsoft.Build.Engine.dll
-protect /Binaries/ThirdParty/Mono/Mac/lib/mono/xbuild/14.0/bin/Microsoft.Build.Framework.dll
-protect /Binaries/ThirdParty/Mono/Mac/lib/mono/xbuild/14.0/bin/Microsoft.Build.Tasks.Core.dll
-protect /Binaries/ThirdParty/Mono/Mac/lib/mono/xbuild/14.0/bin/Microsoft.Build.Utilities.Core.dll
-protect /Binaries/ThirdParty/Mono/Mac/lib/mono/xbuild/14.0/bin/Mono.XBuild.Tasks.dll
-protect /Binaries/ThirdParty/Mono/Mac/lib/mono/monodoc/monodoc.dll
+protect /Binaries/ThirdParty/Mono/Mac/lib/mono/4.5/System.Core.dll
+protect /Binaries/ThirdParty/Mono/Mac/lib/mono/4.5/System.ComponentModel.DataAnnotations.dll
+protect /Binaries/ThirdParty/Mono/Mac/lib/mono/4.5/System.Web.Mvc.dll
+protect /Binaries/ThirdParty/Mono/Mac/lib/mono/4.5/System.ServiceModel.Discovery.dll
+protect /Binaries/ThirdParty/Mono/Mac/lib/mono/4.5/Microsoft.VisualC.dll
+protect /Binaries/ThirdParty/Mono/Mac/lib/mono/4.5/System.Web.ApplicationServices.dll
+protect /Binaries/ThirdParty/Mono/Mac/etc/mono/registry/LocalMachine/
+protect /Binaries/ThirdParty/Mono/Mac/etc/mono/registry/last-btime
 - /Binaries/ThirdParty/ANGLE
 - /Binaries/ThirdParty/IntelISPCTexComp
 + /Binaries/ThirdParty/**
diff --git a/Engine/Source/Programs/UnrealBuildTool/ToolChain/RemoteMac.cs b/Engine/Source/Programs/UnrealBuildTool/ToolChain/RemoteMac.cs
--- a/Engine/Source/Programs/UnrealBuildTool/ToolChain/RemoteMac.cs
+++ b/Engine/Source/Programs/UnrealBuildTool/ToolChain/RemoteMac.cs
@@ -247,6 +247,7 @@ namespace UnrealBuildTool
 			CommonRsyncArguments.Add("--delete"); // Delete anything not in the source directory
 			CommonRsyncArguments.Add("--delete-excluded"); // Delete anything not in the source directory
 			CommonRsyncArguments.Add("--times"); // Preserve modification times
+			CommonRsyncArguments.Add("--omit-dir-times"); // Ignore modification times for directories
 			CommonRsyncArguments.Add("--prune-empty-dirs"); // Remove empty directories from the file list
 
 			// Get the remote base directory
