include all of Eigen module source when copying to remote server via RSync, because some headers do not have identifiable extensions
#rnx

diff --git a/Engine/Build/Rsync/RsyncEngine.txt b/Engine/Build/Rsync/RsyncEngine.txt
--- a/Engine/Build/Rsync/RsyncEngine.txt
+++ b/Engine/Build/Rsync/RsyncEngine.txt
@@ -294,6 +294,7 @@ protect /Saved/
 + /Source/ThirdParty/**.dylib
 + /Source/ThirdParty/**.embeddedframework.zip
 + /Source/ThirdParty/libJPG/*.cpp
++ /Source/ThirdParty/Eigen/**
 - /Source/ThirdParty/**
 
 # Source: Include everything else
