Prevent unzipped frameworks from being removed every time a remote IOS compile is started, preventing iterative builds.

diff --git a/Engine/Build/Rsync/RsyncEngine.txt b/Engine/Build/Rsync/RsyncEngine.txt
--- a/Engine/Build/Rsync/RsyncEngine.txt
+++ b/Engine/Build/Rsync/RsyncEngine.txt
@@ -20,6 +20,8 @@ protect /Intermediate/Build/IOS/**
 protect /Intermediate/Build/TVOS/**
 protect /Intermediate/IOS/**
 protect /Intermediate/Remote/**
+protect /Intermediate/UnzippedFrameworks/**
+protect /Intermediate/ThirdPartyFrameworks/**
 protect /**/Binaries/Mac/**
 protect /**/Binaries/IOS/**
 protect /**/Binaries/TVOS/**
