Fix error generating dsyms on Mac through the remote toolchain.

diff --git a/Engine/Build/Rsync/RsyncEngine.txt b/Engine/Build/Rsync/RsyncEngine.txt
--- a/Engine/Build/Rsync/RsyncEngine.txt
+++ b/Engine/Build/Rsync/RsyncEngine.txt
@@ -9,6 +9,8 @@
 
 # Don't remove build products from previous builds
 + /Binaries/
++ /Binaries/Mac/
++ /Binaries/Mac/DsymExporter
 protect /Binaries/Mac/**
 protect /Binaries/IOS/**
 protect /Binaries/TVOS/**
@@ -299,3 +301,9 @@ protect /Binaries/ThirdParty/Mono/Mac/lib/mono/monodoc/monodoc.dll
 + /Source/ThirdParty/**.dylib
 + /Source/ThirdParty/**.embeddedframework.zip
 + /Source/ThirdParty/libJPG/*.cpp
+
+# Default resource files
++ /Source/Runtime/Launch/
++ /Source/Runtime/Launch/Resources/
++ /Source/Runtime/Launch/Resources/Mac/
++ /Source/Runtime/Launch/Resources/Mac/**
