Fix issues with running DsymExporter remotely from a PC, and reorder engine filters to reduce the amount of junk which has to be re-transferred every time.

diff --git a/Engine/Build/Rsync/RsyncEngine.txt b/Engine/Build/Rsync/RsyncEngine.txt
--- a/Engine/Build/Rsync/RsyncEngine.txt
+++ b/Engine/Build/Rsync/RsyncEngine.txt
@@ -7,36 +7,58 @@
 #  * Directories will not be recursed through unless their parent directories match a rule.
 #  * If directories match a protect rule, it applies to the entire directory and short-circuits traversal of the contents.
 
-# Don't remove build products from previous builds
+# Exclude anything for platforms we don't care about
+- /**/Win32/
+- /**/Win64/
+- /**/Android/
+- /**/Linux/
+- /**/PS4/
+- /**/Switch/
+- /**/XboxOne/
+
+- /**/x64/
+- /**/win32/
+- /**/win64/
+- /**/Linux32/
+- /**/Linux64/
+- /**/linux32/
+- /**/linux64/
+
+# Also ignore tps files; we don't need them to build
+- *.tps
+
+# Binaries
 + /Binaries/
+
+# Binaries/DotNET: Include C# tools
++ /Binaries/DotNET/
++ /Binaries/DotNET/UnrealBuildTool.exe
++ /Binaries/DotNET/UnrealBuildTool.pdb
++ /Binaries/DotNET/DotNETUtilities.dll
++ /Binaries/DotNET/DotNETUtilities.pdb
++ /Binaries/DotNET/RPCUtility.exe
++ /Binaries/DotNET/RPCUtility.pdb
++ /Binaries/DotNET/Ionic.Zip.Reduced.dll
++ /Binaries/DotNET/IOS/
++ /Binaries/DotNET/IOS/**
+
+# Binaries/Mac: Include DsymExporter, ignore everything else
 + /Binaries/Mac/
 + /Binaries/Mac/DsymExporter
 protect /Binaries/Mac/**
-protect /Binaries/IOS/**
-protect /Binaries/TVOS/**
-protect /Build/IOS/*.PackageVersionCounter
-+ /Intermediate/
-+ /Intermediate/Build/
-protect /Intermediate/Build/Mac/**
-protect /Intermediate/Build/IOS/**
-protect /Intermediate/Build/TVOS/**
-protect /Intermediate/IOS/**
-protect /Intermediate/Remote/**
-protect /Intermediate/UnzippedFrameworks/**
-protect /Intermediate/ThirdPartyFrameworks/**
-protect /**/Binaries/Mac/**
-protect /**/Binaries/IOS/**
-protect /**/Binaries/TVOS/**
-protect /**/Intermediate/Build/Mac/**
-protect /**/Intermediate/Build/IOS/**
-protect /**/Intermediate/Build/TVOS/**
-protect /Saved/**
-+ /Programs/
-protect /Programs/UnrealBuildTool/**
-protect /Programs/UnrealHeaderTool/Intermediate/**
-protect /Programs/UnrealHeaderTool/Saved/**
 
-# Don't overwrite symlinks created by SetupMono.sh
+# Binaries/IOS
+protect /Binaries/IOS/
+
+# Binaries/TVOS
+protect /Binaries/TVOS/
+
+# Binaries/ThirdParty. Include Mono for Mac, but don't overwrite symlinks created by SetupMono.sh
++ /Binaries/ThirdParty/
++ /Binaries/ThirdParty/Mono/
++ /Binaries/ThirdParty/Mono/Mac/
++ /Binaries/ThirdParty/Mono/Mac/**
+- /Binaries/ThirdParty/Mono/**
 protect /Binaries/ThirdParty/Mono/Mac/bin/mono
 protect /Binaries/ThirdParty/Mono/Mac/lib/mono/4.5/Accessibility.dll
 protect /Binaries/ThirdParty/Mono/Mac/lib/mono/4.5/Commons.Xml.Relaxng.dll
@@ -185,42 +207,11 @@ protect /Binaries/ThirdParty/Mono/Mac/lib/mono/xbuild/14.0/bin/Microsoft.Build.T
 protect /Binaries/ThirdParty/Mono/Mac/lib/mono/xbuild/14.0/bin/Microsoft.Build.Utilities.Core.dll
 protect /Binaries/ThirdParty/Mono/Mac/lib/mono/xbuild/14.0/bin/Mono.XBuild.Tasks.dll
 protect /Binaries/ThirdParty/Mono/Mac/lib/mono/monodoc/monodoc.dll
+- /Binaries/ThirdParty/ANGLE
+- /Binaries/ThirdParty/IntelISPCTexComp
++ /Binaries/ThirdParty/**
 
-# Exclude anything for platforms we don't care about
-- /**/Win32/
-- /**/Win64/
-- /**/Android/
-- /**/Linux/
-- /**/PS4/
-- /**/Switch/
-- /**/XboxOne/
-
-# Also exclude any C# intermediate directories
-- /Source/Programs/**/obj/
-
-# Add all the paths to sync
-+ /
-+ /Binaries/
-+ /Binaries/DotNET/
-+ /Binaries/DotNET/UnrealBuildTool.exe
-+ /Binaries/DotNET/UnrealBuildTool.pdb
-+ /Binaries/DotNET/DotNETUtilities.dll
-+ /Binaries/DotNET/DotNETUtilities.pdb
-+ /Binaries/DotNET/RPCUtility.exe
-+ /Binaries/DotNET/RPCUtility.pdb
-+ /Binaries/DotNET/Ionic.Zip.Reduced.dll
-+ /Binaries/DotNET/IOS/
-+ /Binaries/DotNET/IOS/**
-+ /Binaries/ThirdParty/
-+ /Binaries/ThirdParty/Ionic/
-+ /Binaries/ThirdParty/Ionic/Ionic.Zip.Reduced.dll
-+ /Binaries/ThirdParty/Ionic/Ionic.Zip.Reduced.Core.dll
-+ /Binaries/ThirdParty/Ionic/Ionic.Zip.Reduced.Core.pdb
-+ /Binaries/ThirdParty/Mono/
-+ /Binaries/ThirdParty/Mono/Mac/
-+ /Binaries/ThirdParty/Mono/Mac/**
-+ /Binaries/ThirdParty/**/
-+ /Binaries/ThirdParty/**.dylib
+# Build
 + /Build/
 + /Build/*
 + /Build/BatchFiles/
@@ -229,81 +220,87 @@ protect /Binaries/ThirdParty/Mono/Mac/lib/mono/monodoc/monodoc.dll
 + /Build/Mac/**
 + /Build/IOS/
 + /Build/IOS/**
+protect /Build/IOS/*.PackageVersionCounter
+
+# Config: Include everything
 + /Config/
 + /Config/**
+
+# Content: Exclude everything except things we need to enumerate as runtime dependencies
 + /Content/
++ /Content/Internationalization/
++ /Content/Internationalization/**
 + /Content/Slate/
 + /Content/Slate/**
 + /Content/SlateDebug/
 + /Content/SlateDebug/**
+- /Content/**
+
+# Intermediate
 + /Intermediate/
 + /Intermediate/Build/
 + /Intermediate/Build/XmlConfigCache.bin
 + /Intermediate/Build/BuildRules/
 + /Intermediate/Build/BuildRules/*.dll
+protect /Intermediate/**
+
+# Plugins
++ /Plugins/
++ /Plugins/**/Binaries/Mac/
++ /Plugins/**/Binaries/IOS/
++ /Plugins/**/Binaries/TVOS/
+- /Plugins/**/Binaries/*/
+- /Plugins/**/Content/
+- /Plugins/**/Resources/
+protect /Plugins/**/Intermediate/
++ /Plugins/**
+
+# Programs
 + /Programs/
+protect /Programs/UnrealBuildTool/
 + /Programs/UnrealHeaderTool/
 + /Programs/UnrealHeaderTool/Config/
-+ /Programs/UnrealHeaderTool/Config/*.ini
-+ /Plugins/
-+ /Plugins/**/
-+ /Plugins/**.uplugin
-+ /Plugins/**/Source/**.h
-+ /Plugins/**/Source/**.cpp
-+ /Plugins/**/Source/**.mm
-+ /Plugins/**/Source/**.m
-+ /Plugins/**/Source/**.cs
-+ /Plugins/**/Source/**.inl
-+ /Plugins/**/Source/**.dylib
-+ /Plugins/**.a
-+ /Plugins/**.embeddedframework.zip
-+ /Plugins/**.xml
++ /Programs/UnrealHeaderTool/Config/**
+protect /Programs/UnrealHeaderTool/Intermediate/
+protect /Programs/UnrealHeaderTool/Saved/
+
+# Saved
+protect /Saved/
+
+# Source
 + /Source/
-+ /Source/*
-+ /Source/*.cs
-+ /Source/**.xml
-+ /Source/**/
-+ /Source/Runtime/**.h
-+ /Source/Runtime/**.cpp
-+ /Source/Runtime/**.mm
-+ /Source/Runtime/**.m
-+ /Source/Runtime/**.cs
-+ /Source/Runtime/**.inl
-+ /Source/Developer/**.h
-+ /Source/Developer/**.cpp
-+ /Source/Developer/**.mm
-+ /Source/Developer/**.m
-+ /Source/Developer/**.cs
-+ /Source/Developer/**.inl
-+ /Source/Editor/**.h
-+ /Source/Editor/**.cpp
-+ /Source/Editor/**.mm
-+ /Source/Editor/**.m
-+ /Source/Editor/**.cs
-+ /Source/Editor/**.inl
-+ /Source/Programs/**.h
-+ /Source/Programs/**.cpp
-+ /Source/Programs/**.mm
-+ /Source/Programs/**.m
-+ /Source/Programs/**.cs
-+ /Source/Programs/**.inl
-+ /Source/Programs/**.def
-+ /Source/Programs/**.csproj
-+ /Source/ThirdParty/
+
+# Source/Runtime: Blacklist specific folders
+- /Source/Runtime/Navmesh/RecastDemo/
+
+# Source/Programs: Whitelist specific programs
++ /Source/Programs/
++ /Source/Programs/AutomationTool/
++ /Source/Programs/DotNETCommon/
++ /Source/Programs/UnrealBuildTool/
++ /Source/Programs/UnrealPak/
++ /Source/Programs/UnrealHeaderTool/
++ /Source/Programs/UnrealHeaderTool/**.def
+- /Source/Programs/**/bin/
+- /Source/Programs/**/obj/
+- /Source/Programs/*/
+
+# Source/ThirdParty: Whitelist specific file types
 + /Source/ThirdParty/**/
 + /Source/ThirdParty/**.cs
 + /Source/ThirdParty/**.h
 + /Source/ThirdParty/**.hh
 + /Source/ThirdParty/**.hpp
++ /Source/ThirdParty/**.m
 + /Source/ThirdParty/**.mm
++ /Source/ThirdParty/**.cs
 + /Source/ThirdParty/**.inl
 + /Source/ThirdParty/**.a
 + /Source/ThirdParty/**.dylib
 + /Source/ThirdParty/**.embeddedframework.zip
 + /Source/ThirdParty/libJPG/*.cpp
+- /Source/ThirdParty/**
 
-# Default resource files
-+ /Source/Runtime/Launch/
-+ /Source/Runtime/Launch/Resources/
-+ /Source/Runtime/Launch/Resources/Mac/
-+ /Source/Runtime/Launch/Resources/Mac/**
+# Source: Include everything else
++ /Source/**/
++ /Source/**
diff --git a/Engine/Build/Rsync/RsyncEngineScripts.txt b/Engine/Build/Rsync/RsyncEngineScripts.txt
--- a/Engine/Build/Rsync/RsyncEngineScripts.txt
+++ b/Engine/Build/Rsync/RsyncEngineScripts.txt
@@ -1,5 +1,5 @@
-; This file contains a list of shell scripts to be synced to the remote. These files will be locally converted to use Unix line endings, and will be
-; marked as executable on the remote. Wildcards are not supported.
+# This file contains a list of shell scripts to be synced to the remote. These files will be locally converted to use Unix line endings, and will be
+# marked as executable on the remote. Wildcards are not supported.
 
 Binaries/ThirdParty/Mono/Mac/bin/xbuild
 Build/BatchFiles/Mac/Build.sh
diff --git a/Engine/Source/Programs/UnrealBuildTool/ToolChain/RemoteMac.cs b/Engine/Source/Programs/UnrealBuildTool/ToolChain/RemoteMac.cs
--- a/Engine/Source/Programs/UnrealBuildTool/ToolChain/RemoteMac.cs
+++ b/Engine/Source/Programs/UnrealBuildTool/ToolChain/RemoteMac.cs
@@ -109,6 +109,11 @@ namespace UnrealBuildTool
 		/// <summary>
 		/// Arguments that are used by every Rsync call
 		/// </summary>
+		private List<string> BasicRsyncArguments;
+
+		/// <summary>
+		/// Arguments that are used by directory Rsync call
+		/// </summary>
 		private List<string> CommonRsyncArguments;
 
 		/// <summary>
@@ -227,16 +232,19 @@ namespace UnrealBuildTool
 			CommonSshArguments.Add(String.Format("\"{0}@{1}\"", UserName, ServerName));
 
 			// Build a list of arguments for Rsync
-			CommonRsyncArguments = new List<string>();
-			CommonRsyncArguments.Add("--compress");
+			BasicRsyncArguments = new List<string>();
+			BasicRsyncArguments.Add("--compress");
+			BasicRsyncArguments.Add("--verbose");
+			BasicRsyncArguments.Add(String.Format("--rsh=\"{0} -p {1}\"", RsyncAuthentication, ServerPort));
+
+			// Build a list of arguments for Rsync filters
+			CommonRsyncArguments = new List<string>(BasicRsyncArguments);
 			CommonRsyncArguments.Add("--recursive");
 			CommonRsyncArguments.Add("--delete"); // Delete anything not in the source directory
 			CommonRsyncArguments.Add("--delete-excluded"); // Delete anything not in the source directory
 			CommonRsyncArguments.Add("--times"); // Preserve modification times
-			CommonRsyncArguments.Add("--verbose");
-			CommonRsyncArguments.Add("-m");
+			CommonRsyncArguments.Add("--prune-empty-dirs"); // Remove empty directories from the file list
 			CommonRsyncArguments.Add("--chmod=ug=rwX,o=rxX");
-			CommonRsyncArguments.Add(String.Format("--rsh=\"{0} -p {1}\"", RsyncAuthentication, ServerPort));
 
 			// Get the remote base directory
 			StringBuilder Output;
@@ -458,7 +466,6 @@ namespace UnrealBuildTool
 			}
 
 			// Upload the workspace files
-			Log.TraceInformation("[Remote] Uploading workspace files");
 			UploadWorkspace(TempDir);
 
 			// Fixup permissions on any shell scripts
@@ -629,6 +636,27 @@ namespace UnrealBuildTool
 			}
 		}
 
+		/// <summary>
+		/// Upload a single file to the remote
+		/// </summary>
+		/// <param name="LocalDirectory">The base directory to copy</param>
+		/// <param name="RemoteDirectory">The remote directory</param>
+		/// <param name="LocalFileList">The file to upload</param>
+		void UploadFiles(DirectoryReference LocalDirectory, string RemoteDirectory, FileReference LocalFileList)
+		{
+			List<string> Arguments = new List<string>(BasicRsyncArguments);
+			Arguments.Add(String.Format("--files-from=\"{0}\"", GetLocalCygwinPath(LocalFileList)));
+			Arguments.Add(String.Format("\"{0}/\"", GetLocalCygwinPath(LocalDirectory)));
+			Arguments.Add(String.Format("\"{0}@{1}\":'{2}/'", UserName, ServerName, RemoteDirectory));
+			Arguments.Add("-q");
+
+			int Result = Rsync(String.Join(" ", Arguments));
+			if(Result != 0)
+			{
+				throw new BuildException("Error while running Rsync (exit code {0})", Result);
+			}
+		}
+
 		/// <summary>
 		/// Upload a single directory to the remote
 		/// </summary>
@@ -688,7 +716,7 @@ namespace UnrealBuildTool
 			foreach(string Line in FileReference.ReadAllLines(ScriptPathsFileName))
 			{
 				string FileToUpload = Line.Trim();
-				if(FileToUpload.Length > 0 && FileToUpload[0] != ';')
+				if(FileToUpload.Length > 0 && FileToUpload[0] != '#')
 				{
 					ScriptPaths.Add(FileToUpload);
 				}
@@ -714,21 +742,6 @@ namespace UnrealBuildTool
 				TargetFiles.Add(TargetFile);
 			}
 
-			// Write a file that protects all the scripts from being overridden by the standard engine filters
-			FileReference ScriptUploadList = FileReference.Combine(TempDir, "RsyncEngineScripts-Upload.txt");
-			using(StreamWriter Writer = new StreamWriter(ScriptUploadList.FullName))
-			{
-				foreach(string ScriptPath in ScriptPaths)
-				{
-					for(int SlashIdx = ScriptPath.IndexOf('/', 1); SlashIdx != -1; SlashIdx = ScriptPath.IndexOf('/', SlashIdx + 1))
-					{
-						Writer.WriteLine("+ {0}", ScriptPath.Substring(0, SlashIdx));
-					}
-					Writer.WriteLine("+ {0}", ScriptPath);
-				}
-				Writer.WriteLine("protect *");
-			}
-
 			// Write a file that protects all the scripts from being overridden by the standard engine filters
 			FileReference ScriptProtectList = FileReference.Combine(TempDir, "RsyncEngineScripts-Protect.txt");
 			using(StreamWriter Writer = new StreamWriter(ScriptProtectList.FullName))
@@ -740,14 +753,15 @@ namespace UnrealBuildTool
 			}
 
 			// Upload these files to the remote
-			List<FileReference> FilterLocations = new List<FileReference>();
-			FilterLocations.Add(ScriptUploadList);
-			UploadDirectory(TempDir, GetRemotePath(UnrealBuildTool.EngineDirectory), FilterLocations);
+			Log.TraceInformation("[Remote] Uploading scripts...");
+			UploadFiles(TempDir, GetRemotePath(UnrealBuildTool.EngineDirectory), ScriptPathsFileName);
 
 			// Upload the engine files
 			List<FileReference> EngineFilters = new List<FileReference>();
 			EngineFilters.Add(ScriptProtectList);
 			EngineFilters.Add(FileReference.Combine(UnrealBuildTool.EngineDirectory, "Build", "Rsync", "RsyncEngine.txt"));
+
+			Log.TraceInformation("[Remote] Uploading engine files...");
 			UploadDirectory(UnrealBuildTool.EngineDirectory, GetRemotePath(UnrealBuildTool.EngineDirectory), EngineFilters);
 
 			// Upload the project files
@@ -762,6 +776,7 @@ namespace UnrealBuildTool
 				}
 				ProjectFilters.Add(FileReference.Combine(UnrealBuildTool.EngineDirectory, "Build", "Rsync", "RsyncProject.txt"));
 
+				Log.TraceInformation("[Remote] Uploading project files...");
 				UploadDirectory(ProjectFile.Directory, GetRemotePath(ProjectFile.Directory), ProjectFilters);
 			}
 		}
