UBT: Fix embedded framework zips not being uploaded for plugins.

diff --git a/Engine/Build/Rsync/RsyncEngine.txt b/Engine/Build/Rsync/RsyncEngine.txt
--- a/Engine/Build/Rsync/RsyncEngine.txt
+++ b/Engine/Build/Rsync/RsyncEngine.txt
@@ -241,6 +241,7 @@ protect /Binaries/ThirdParty/Mono/Mac/lib/mono/monodoc/monodoc.dll
 + /Plugins/**/Source/**.m
 + /Plugins/**/Source/**.cs
 + /Plugins/**/Source/**.inl
++ /Plugins/**.embeddedframework.zip
 + /Source/
 + /Source/*
 + /Source/*.cs
diff --git a/Engine/Build/Rsync/RsyncProject.txt b/Engine/Build/Rsync/RsyncProject.txt
--- a/Engine/Build/Rsync/RsyncProject.txt
+++ b/Engine/Build/Rsync/RsyncProject.txt
@@ -23,6 +23,6 @@ protect /Intermediate/**
 + /Plugins/**/*.uplugin
 + /Plugins/**/Config/**
 + /Plugins/**/Source/**
++ /Plugins/**.embeddedframework.zip
 + /Source/
 + /Source/**
-
