Upload all static libraries in plugin folders as part of remote builds.

diff --git a/Engine/Build/Rsync/RsyncEngine.txt b/Engine/Build/Rsync/RsyncEngine.txt
--- a/Engine/Build/Rsync/RsyncEngine.txt
+++ b/Engine/Build/Rsync/RsyncEngine.txt
@@ -241,6 +241,7 @@ protect /Binaries/ThirdParty/Mono/Mac/lib/mono/monodoc/monodoc.dll
 + /Plugins/**/Source/**.m
 + /Plugins/**/Source/**.cs
 + /Plugins/**/Source/**.inl
++ /Plugins/**.a
 + /Plugins/**.embeddedframework.zip
 + /Source/
 + /Source/*
