Allow mirroring Oodle files to remote Mac.

diff --git a/Engine/Build/Rsync/RsyncProject.txt b/Engine/Build/Rsync/RsyncProject.txt
--- a/Engine/Build/Rsync/RsyncProject.txt
+++ b/Engine/Build/Rsync/RsyncProject.txt
@@ -19,6 +19,8 @@ protect /Intermediate/**
 + /Config/
 + /Config/**
 + /Content/
++ /Content/Oodle/
++ /Content/Oodle/**
 + /Content/Slate/
 + /Content/Slate/**
 + /Content/SlateDebug/
