RSync fixes for xml files

diff --git a/Engine/Build/Rsync/RsyncEngine.txt b/Engine/Build/Rsync/RsyncEngine.txt
--- a/Engine/Build/Rsync/RsyncEngine.txt
+++ b/Engine/Build/Rsync/RsyncEngine.txt
@@ -248,9 +248,11 @@ protect /Binaries/ThirdParty/Mono/Mac/lib/mono/monodoc/monodoc.dll
 + /Plugins/**/Source/**.inl
 + /Plugins/**.a
 + /Plugins/**.embeddedframework.zip
++ /Plugins/**.xml
 + /Source/
 + /Source/*
 + /Source/*.cs
++ /Source/**.xml
 + /Source/**/
 + /Source/Runtime/**.h
 + /Source/Runtime/**.cpp
