#DEV-VR - Fixing remote iOS compile issue with BP projects that require plugins

diff --git a/Engine/Build/Rsync/RsyncProject.txt b/Engine/Build/Rsync/RsyncProject.txt
--- a/Engine/Build/Rsync/RsyncProject.txt
+++ b/Engine/Build/Rsync/RsyncProject.txt
@@ -7,6 +7,8 @@
 + /Intermediate/Build/
 + /Intermediate/Build/BuildRules/
 + /Intermediate/Build/BuildRules/*.dll
++ /Intermediate/Source/
++ /Intermediate/Source/**
 
 protect /Intermediate/**
 
