Fix for missing files in remote iOS builds

diff --git a/Engine/Build/Rsync/RsyncEngine.txt b/Engine/Build/Rsync/RsyncEngine.txt
--- a/Engine/Build/Rsync/RsyncEngine.txt
+++ b/Engine/Build/Rsync/RsyncEngine.txt
@@ -12,7 +12,9 @@
 - /**/Win64/
 - /Plugins/**/Android/
 - /Source/ThirdParty/**/Android/
-- /**/Linux/
+- /**/Lumin/
++ /Source/Programs/UnrealBuildTool/Platform/Linux/LinuxCommon.cs
+- /**/Linux/**
 - /**/PS4/
 - /**/Switch/
 - /**/XboxOne/
@@ -37,6 +39,8 @@
 + /Binaries/DotNET/AgentInterface.dll
 + /Binaries/DotNET/UnrealBuildTool.exe
 + /Binaries/DotNET/UnrealBuildTool.pdb
++ /Binaries/DotNET/BuildUtilities.dll
++ /Binaries/DotNET/BuildUtilities.pdb
 + /Binaries/DotNET/DotNETUtilities.dll
 + /Binaries/DotNET/DotNETUtilities.pdb
 + /Binaries/DotNET/Ionic.Zip.Reduced.dll
