Fix for Lumin/Linux .cs files being deleted that are required

diff --git a/Engine/Build/Rsync/RsyncEngine.txt b/Engine/Build/Rsync/RsyncEngine.txt
--- a/Engine/Build/Rsync/RsyncEngine.txt
+++ b/Engine/Build/Rsync/RsyncEngine.txt
@@ -7,13 +7,18 @@
 #  * Directories will not be recursed through unless their parent directories match a rule.
 #  * If directories match a protect rule, it applies to the entire directory and short-circuits traversal of the contents.
 
+# These platform files are relied on by BuildTool so must be included (even though platform is irrelevant)
++ /Source/Programs/UnrealBuildTool/Platform/Lumin/**
++ /Source/Programs/UnrealBuildTool/Platform/Linux/**
+protect /Source/Programs/UnrealBuildTool/Platform/Lumin/**
+protect /Source/Programs/UnrealBuildTool/Platform/Linux/**
+
 # Exclude anything for platforms we don't care about
 - /**/Win32/
 - /**/Win64/
 - /Plugins/**/Android/
 - /Source/ThirdParty/**/Android/
-- /**/Lumin/
-+ /Source/Programs/UnrealBuildTool/Platform/Linux/LinuxCommon.cs
+- /**/Lumin/**
 - /**/Linux/**
 - /**/PS4/
 - /**/Switch/
