Fix up iOS remote build after move from Mono to dotnet
- build UBT/UAT on remote Mac
- exclude unnecessary UBT platform source from rsync upload
- modifications to UBT to allow it to compile with Linux and Lumin excluded
- manually convert CRLF for Mac shell scripts after rsync upload

diff --git a/Engine/Build/BatchFiles/Mac/Build.sh b/Engine/Build/BatchFiles/Mac/Build.sh
--- a/Engine/Build/BatchFiles/Mac/Build.sh
+++ b/Engine/Build/BatchFiles/Mac/Build.sh
@@ -4,27 +4,24 @@
 
 cd "`dirname "$0"`/../../../.."
 
-# Setup Environment and Mono
+# Setup Environment for DotNET
 source Engine/Build/BatchFiles/Mac/SetupEnvironment.sh -dotnet Engine/Build/BatchFiles/Mac
 
-# Skip UBT and SWC compile step if we're coming in on an SSH connection (ie remote toolchain)
-if [ -z "$SSH_CONNECTION" ]; then
-	# First make sure that the UnrealBuildTool is up-to-date
-	if ! dotnet build Engine/Source/Programs/UnrealBuildTool/UnrealBuildTool.csproj -c Development -v quiet; then
-		echo "Failed to build to build tool (UnrealBuildTool)"
-		exit 1
-	fi
-
-	# build SCW if specified
-	for i in "$@" ; do
-		if [[ $i == "-buildscw" ]] ; then
-			echo Building ShaderCompileWorker...
-			Engine/Binaries/DotNET/UnrealBuildTool/UnrealBuildTool ShaderCompileWorker Mac Development
-			break
-		fi
-	done
+# First make sure that the UnrealBuildTool is up-to-date
+if ! dotnet build Engine/Source/Programs/UnrealBuildTool/UnrealBuildTool.csproj -c Development -v quiet; then
+	echo "Failed to build to build tool (UnrealBuildTool)"
+	exit 1
 fi
 
+# build SCW if specified
+for i in "$@" ; do
+	if [[ $i == "-buildscw" ]] ; then
+		echo Building ShaderCompileWorker...
+		Engine/Binaries/DotNET/UnrealBuildTool/UnrealBuildTool ShaderCompileWorker Mac Development
+		break
+	fi
+done
+
 echo Running Engine/Binaries/DotNET/UnrealBuildTool/UnrealBuildTool "$@"
 Engine/Binaries/DotNET/UnrealBuildTool/UnrealBuildTool "$@"
 
diff --git a/Engine/Build/Rsync/RsyncEngine.txt b/Engine/Build/Rsync/RsyncEngine.txt
--- a/Engine/Build/Rsync/RsyncEngine.txt
+++ b/Engine/Build/Rsync/RsyncEngine.txt
@@ -7,18 +7,13 @@
 #  * Directories will not be recursed through unless their parent directories match a rule.
 #  * If directories match a protect rule, it applies to the entire directory and short-circuits traversal of the contents.
 
-# These platform files are relied on by BuildTool so must be included (even though platform is irrelevant)
-+ /Source/Programs/UnrealBuildTool/Platform/Lumin/**
-+ /Source/Programs/UnrealBuildTool/Platform/Linux/**
-protect /Source/Programs/UnrealBuildTool/Platform/Lumin/**
-protect /Source/Programs/UnrealBuildTool/Platform/Linux/**
-
 # Exclude anything for platforms we don't care about
 - /**/Win32/
 - /**/Win64/
 - /Plugins/**/Android/
 - /Source/ThirdParty/**/Android/
 - /**/Lumin/**
++ /Source/Programs/UnrealBuildTool/Platform/Linux/LinuxCommon.cs
 - /**/Linux/**
 - /**/PS4/
 - /**/Switch/
@@ -39,11 +34,9 @@ protect /Source/Programs/UnrealBuildTool/Platform/Linux/**
 # Binaries
 + /Binaries/
 
-# Binaries/DotNET: Include C# tools
+# Binaries/DotNET and UnrealBuildTool
 + /Binaries/DotNET/
 + /Binaries/DotNET/AgentInterface.dll
-+ /Binaries/DotNET/UnrealBuildTool.exe
-+ /Binaries/DotNET/UnrealBuildTool.pdb
 + /Binaries/DotNET/BuildUtilities.dll
 + /Binaries/DotNET/BuildUtilities.pdb
 + /Binaries/DotNET/DotNETUtilities.dll
@@ -51,6 +44,7 @@ protect /Source/Programs/UnrealBuildTool/Platform/Linux/**
 + /Binaries/DotNET/Ionic.Zip.Reduced.dll
 + /Binaries/DotNET/IOS/
 + /Binaries/DotNET/IOS/**
+protect /Binaries/DotNET/UnrealBuildTool/**
 
 # Binaries/Mac: Include DsymExporter, ignore everything else
 + /Binaries/Mac/
@@ -224,6 +218,7 @@ protect /Binaries/ThirdParty/Mono/Mac/etc/mono/registry/last-btime
 + /Build/TVOS/
 + /Build/TVOS/**
 protect /Build/IOS/*.PackageVersionCounter
+protect /Build/BatchFiles/Mac/*.crlf
 
 # Config: Include everything
 + /Config/
@@ -285,6 +280,9 @@ protect /Saved/
 + /Source/Programs/AutomationTool/
 + /Source/Programs/DotNETCommon/
 + /Source/Programs/UnrealBuildTool/
++ /Source/Programs/Shared/
++ /Source/Programs/Shared/EpicGames.Core/
+- /Source/Programs/Shared/*/
 + /Source/Programs/UnrealPak/
 + /Source/Programs/UnrealHeaderTool/
 + /Source/Programs/UnrealHeaderTool/**.def
@@ -292,6 +290,12 @@ protect /Saved/
 - /Source/Programs/**/bin/
 - /Source/Programs/**/obj/
 - /Source/Programs/*/
+# these are built on the remote Mac
+protect /Source/Programs/DotNETCommon/obj/**
+protect /Source/Programs/DotNETCommon/BuildUtilities/bin/**
+protect /Source/Programs/UnrealBuildTool/obj/**
+protect /Source/Programs/Shared/EpicGames.Core/obj/**
+protect /Source/Programs/Shared/EpicGames.Core/bin/**
 
 # Source/ThirdParty: Whitelist specific file types
 + /Source/ThirdParty/**/
diff --git a/Engine/Source/Programs/DotNETCommon/BuildUtilities/UEBuildPlatformSDK.cs b/Engine/Source/Programs/DotNETCommon/BuildUtilities/UEBuildPlatformSDK.cs
--- a/Engine/Source/Programs/DotNETCommon/BuildUtilities/UEBuildPlatformSDK.cs
+++ b/Engine/Source/Programs/DotNETCommon/BuildUtilities/UEBuildPlatformSDK.cs
@@ -1163,6 +1163,15 @@ namespace EpicGames.Core
 			}
 		}
 
+		/// <summary>
+		/// Allows the platform to optionally returns a path to the internal SDK
+		/// </summary>
+		/// <returns>Valid path to the internal SDK, null otherwise</returns>
+		public virtual string GetInternalSDKPath()
+		{
+			return null;
+		}
+
 		#endregion
 
 		#region public AutoSDKs Utility
diff --git a/Engine/Source/Programs/UnrealBuildTool/Configuration/PlatformTargetRules.cs b/Engine/Source/Programs/UnrealBuildTool/Configuration/PlatformTargetRules.cs
--- a/Engine/Source/Programs/UnrealBuildTool/Configuration/PlatformTargetRules.cs
+++ b/Engine/Source/Programs/UnrealBuildTool/Configuration/PlatformTargetRules.cs
@@ -9,6 +9,68 @@ using System.Threading.Tasks;
 
 namespace UnrealBuildTool
 {
+	/// <summary>
+	/// Stub partial class for Linux-specific target settings. 
+	/// This class is not in a restricted location to simplify code paths in UBT. It is visible to all UE4 users, without NDA, and will appear
+	/// empty to those without the Linux implementation code.
+	/// </summary>
+	public partial class LinuxTargetRules
+	{
+	}
+
+	/// <summary>
+	/// Stub read-only wrapper for Linux-specific target settings.
+	/// This class is not in a restricted location to simplify code paths in UBT. It is visible to all UE4 users, without NDA, and will appear
+	/// empty to those without the Linux implementation code.
+	/// </summary>
+	public partial class ReadOnlyLinuxTargetRules
+	{
+		/// <summary>
+		/// The private mutable settings object
+		/// </summary>
+		protected LinuxTargetRules Inner;
+
+		/// <summary>
+		/// Constructor
+		/// </summary>
+		/// <param name="Inner">The settings object to wrap</param>
+		public ReadOnlyLinuxTargetRules(LinuxTargetRules Inner)
+		{
+			this.Inner = Inner;
+		}
+	}
+
+	/// <summary>
+	/// Stub partial class for Lumin-specific target settings. 
+	/// This class is not in a restricted location to simplify code paths in UBT. It is visible to all UE4 users, without NDA, and will appear
+	/// empty to those without the Lumin implementation code.
+	/// </summary>
+	public partial class LuminTargetRules
+	{
+	}
+
+	/// <summary>
+	/// Stub read-only wrapper for Lumin-specific target settings.
+	/// This class is not in a restricted location to simplify code paths in UBT. It is visible to all UE4 users, without NDA, and will appear
+	/// empty to those without the Lumin implementation code.
+	/// </summary>
+	public partial class ReadOnlyLuminTargetRules
+	{
+		/// <summary>
+		/// The private mutable settings object
+		/// </summary>
+		protected LuminTargetRules Inner;
+
+		/// <summary>
+		/// Constructor
+		/// </summary>
+		/// <param name="Inner">The settings object to wrap</param>
+		public ReadOnlyLuminTargetRules(LuminTargetRules Inner)
+		{
+			this.Inner = Inner;
+		}
+	}
+
 	/// <summary>
 	/// Stub partial class for XboxOne-specific target settings. 
 	/// This class is not in a restricted location to simplify code paths in UBT. It is visible to all UE4 users, without NDA, and will appear
diff --git a/Engine/Source/Programs/UnrealBuildTool/Platform/Linux/LinuxCommon.cs b/Engine/Source/Programs/UnrealBuildTool/Platform/Linux/LinuxCommon.cs
--- a/Engine/Source/Programs/UnrealBuildTool/Platform/Linux/LinuxCommon.cs
+++ b/Engine/Source/Programs/UnrealBuildTool/Platform/Linux/LinuxCommon.cs
@@ -40,7 +40,7 @@ namespace UnrealBuildTool
 
 		public static string WhichClang()
 		{
-			string InternalSDKPath = ((LinuxPlatformSDK)UEBuildPlatform.GetSDK(UnrealTargetPlatform.Linux)).GetInternalSDKPath();
+			string InternalSDKPath = UEBuildPlatform.GetSDK(UnrealTargetPlatform.Linux).GetInternalSDKPath();
 			if (!String.IsNullOrEmpty(InternalSDKPath))
 			{
 				return Path.Combine(InternalSDKPath, "bin", "clang++");
diff --git a/Engine/Source/Programs/UnrealBuildTool/Platform/Linux/LinuxPlatformSDK.cs b/Engine/Source/Programs/UnrealBuildTool/Platform/Linux/LinuxPlatformSDK.cs
--- a/Engine/Source/Programs/UnrealBuildTool/Platform/Linux/LinuxPlatformSDK.cs
+++ b/Engine/Source/Programs/UnrealBuildTool/Platform/Linux/LinuxPlatformSDK.cs
@@ -120,7 +120,7 @@ namespace UnrealBuildTool
 		/// Returns a path to the internal SDK
 		/// </summary>
 		/// <returns>Valid path to the internal SDK, null otherwise</returns>
-		public string GetInternalSDKPath()
+		public override string GetInternalSDKPath()
 		{
 			string SDKRoot = Environment.GetEnvironmentVariable(SDKRootEnvVar);
 			if (!String.IsNullOrEmpty(SDKRoot))
diff --git a/Engine/Source/Programs/UnrealBuildTool/Platform/Linux/UEBuildLinux.cs b/Engine/Source/Programs/UnrealBuildTool/Platform/Linux/UEBuildLinux.cs
--- a/Engine/Source/Programs/UnrealBuildTool/Platform/Linux/UEBuildLinux.cs
+++ b/Engine/Source/Programs/UnrealBuildTool/Platform/Linux/UEBuildLinux.cs
@@ -31,7 +31,7 @@ namespace UnrealBuildTool
 	/// <summary>
 	/// Linux-specific target settings
 	/// </summary>
-	public class LinuxTargetRules
+	public partial class LinuxTargetRules
 	{
 		/// <summary>
 		/// Constructor
@@ -85,22 +85,8 @@ namespace UnrealBuildTool
 	/// <summary>
 	/// Read-only wrapper for Linux-specific target settings
 	/// </summary>
-	public class ReadOnlyLinuxTargetRules
+	public partial class ReadOnlyLinuxTargetRules
 	{
-		/// <summary>
-		/// The private mutable settings object
-		/// </summary>
-		private LinuxTargetRules Inner;
-
-		/// <summary>
-		/// Constructor
-		/// </summary>
-		/// <param name="Inner">The settings object to wrap</param>
-		public ReadOnlyLinuxTargetRules(LinuxTargetRules Inner)
-		{
-			this.Inner = Inner;
-		}
-
 		/// <summary>
 		/// Accessors for fields on the inner TargetRules instance
 		/// </summary>
diff --git a/Engine/Source/Programs/UnrealBuildTool/Platform/Lumin/UEBuildLumin.cs b/Engine/Source/Programs/UnrealBuildTool/Platform/Lumin/UEBuildLumin.cs
--- a/Engine/Source/Programs/UnrealBuildTool/Platform/Lumin/UEBuildLumin.cs
+++ b/Engine/Source/Programs/UnrealBuildTool/Platform/Lumin/UEBuildLumin.cs
@@ -17,7 +17,7 @@ namespace UnrealBuildTool
 	/// <summary>
 	/// Lumin-specific target settings
 	/// </summary>	
-	public class LuminTargetRules
+	public partial class LuminTargetRules
 	{
 		/// <summary>
 		/// Lists GPU Architectures that you want to build (mostly used for mobile etc.)
@@ -35,22 +35,8 @@ namespace UnrealBuildTool
 	/// <summary>
 	/// Read-only wrapper for Android-specific target settings
 	/// </summary>
-	public class ReadOnlyLuminTargetRules
+	public partial class ReadOnlyLuminTargetRules
 	{
-		/// <summary>
-		/// The private mutable settings object
-		/// </summary>
-		private LuminTargetRules Inner;
-
-		/// <summary>
-		/// Constructor
-		/// </summary>
-		/// <param name="Inner">The settings object to wrap</param>
-		public ReadOnlyLuminTargetRules(LuminTargetRules Inner)
-		{
-			this.Inner = Inner;
-		}
-
 		/// <summary>
 		/// Accessors for fields on the inner TargetRules instance
 		/// </summary>
diff --git a/Engine/Source/Programs/UnrealBuildTool/ToolChain/RemoteMac.cs b/Engine/Source/Programs/UnrealBuildTool/ToolChain/RemoteMac.cs
--- a/Engine/Source/Programs/UnrealBuildTool/ToolChain/RemoteMac.cs
+++ b/Engine/Source/Programs/UnrealBuildTool/ToolChain/RemoteMac.cs
@@ -972,6 +972,9 @@ namespace UnrealBuildTool
 				Execute("/", String.Format("rm -rf {0}/Intermediate/TVOS/*.plist", GetRemotePath(ProjectFile.Directory)));
 			}
 
+			// Convert CRLF to LF for all shell scripts
+			Execute(RemoteBaseDir, String.Format("for i in {0}/Build/BatchFiles/Mac/*.sh; do mv $i $i.crlf; tr -d '\r' < $i.crlf > $i; done", EscapeShellArgument(GetRemotePath(UnrealBuildTool.EngineDirectory))));
+
 			// Fixup permissions on any shell scripts
 			Execute(RemoteBaseDir, String.Format("chmod +x {0}/Build/BatchFiles/Mac/*.sh", EscapeShellArgument(GetRemotePath(UnrealBuildTool.EngineDirectory))));
 		}
