Fix up missing requirements for remote iOS build
Fix up libimobiledevice not specifying exe extenstion to CreateProcess

diff --git a/Engine/Build/Rsync/RsyncEngine.txt b/Engine/Build/Rsync/RsyncEngine.txt
--- a/Engine/Build/Rsync/RsyncEngine.txt
+++ b/Engine/Build/Rsync/RsyncEngine.txt
@@ -56,7 +56,7 @@ protect /Binaries/IOS/
 # Binaries/TVOS
 protect /Binaries/TVOS/
 
-# Binaries/ThirdParty, just DotNet and VisualStudio required to run UnrealBuildTool
+# Binaries/ThirdParty, just DotNet and VisualStudio required to run UnrealBuildTool, and libtbb used by UHT
 + /Binaries/ThirdParty/
 + /Binaries/ThirdParty/DotNet/
 + /Binaries/ThirdParty/DotNet/Mac/
@@ -64,7 +64,10 @@ protect /Binaries/TVOS/
 - /Binaries/ThirdParty/DotNet/**
 + /Binaries/ThirdParty/VisualStudio/
 + /Binaries/ThirdParty/VisualStudio/**
-
++ /Binaries/ThirdParty/Intel/
++ /Binaries/ThirdParty/Intel/TBB/
++ /Binaries/ThirdParty/Intel/TBB/Mac/
++ /Binaries/ThirdParty/Intel/TBB/Mac/libtbb.dylib
 
 # Build
 + /Build/
@@ -143,6 +146,7 @@ protect /Saved/
 + /Source/Programs/UnrealBuildTool/
 + /Source/Programs/Shared/
 + /Source/Programs/Shared/EpicGames.Core/
++ /Source/Programs/Shared/EpicGames.Build/
 - /Source/Programs/Shared/*/
 + /Source/Programs/UnrealPak/
 + /Source/Programs/UnrealHeaderTool/
@@ -157,6 +161,8 @@ protect /Source/Programs/DotNETCommon/BuildUtilities/bin/**
 protect /Source/Programs/UnrealBuildTool/obj/**
 protect /Source/Programs/Shared/EpicGames.Core/obj/**
 protect /Source/Programs/Shared/EpicGames.Core/bin/**
+protect /Source/Programs/Shared/EpicGames.Build/obj/**
+protect /Source/Programs/Shared/EpicGames.Build/bin/**
 
 # Source/ThirdParty: Whitelist specific file types
 + /Source/ThirdParty/**/
@@ -174,6 +180,7 @@ protect /Source/Programs/Shared/EpicGames.Core/bin/**
 + /Source/ThirdParty/libJPG/*.cpp
 + /Source/ThirdParty/libSampleRate/**.cpp
 + /Source/ThirdParty/Eigen/**
++ /Source/ThirdParty/mimalloc/src/**
 - /Source/ThirdParty/**
 
 # Source: Include everything else
diff --git a/Engine/Source/Developer/IOS/IOSTargetPlatform/Private/IOSDeviceHelper.cpp b/Engine/Source/Developer/IOS/IOSTargetPlatform/Private/IOSDeviceHelper.cpp
--- a/Engine/Source/Developer/IOS/IOSTargetPlatform/Private/IOSDeviceHelper.cpp
+++ b/Engine/Source/Developer/IOS/IOSTargetPlatform/Private/IOSDeviceHelper.cpp
@@ -49,6 +49,10 @@ DEFINE_LOG_CATEGORY_STATIC(LogIOSDeviceHelper, Log, All);
         UE_LOG(LogIOSDeviceHelper, Error, TEXT("The current platform is unsupported by Libimobile library."));
 #endif
         ToReturn += ExeName;
+#if PLATFORM_WINDOWS
+		ToReturn += TEXT(".exe");
+#endif
+
         return ToReturn;
     }
 
