Fix iOS Remote toolchain after refactor in CL 18271725 removed Linux partial classes

diff --git a/Engine/Build/Rsync/RsyncEngine.txt b/Engine/Build/Rsync/RsyncEngine.txt
--- a/Engine/Build/Rsync/RsyncEngine.txt
+++ b/Engine/Build/Rsync/RsyncEngine.txt
@@ -12,7 +12,9 @@
 - /**/Win64/
 - /Plugins/**/Android/
 - /Source/ThirdParty/**/Android/
-+ /Source/Programs/UnrealBuildTool/Platform/Linux/LinuxCommon.cs
+# Linux is a hardcoded platform so UBT needs to know about it
++ /Source/Programs/UnrealBuildTool/Platform/Linux/
++ /Source/Programs/UnrealBuildTool/Platform/Linux/**
 - /**/Linux/**
 - /**/PS4/
 - /**/Switch/
