Fix remote toolchain for iOS/tvOS for Shared/EpicGames.UHT change
#rnx

diff --git a/Engine/Build/Rsync/RsyncEngine.txt b/Engine/Build/Rsync/RsyncEngine.txt
--- a/Engine/Build/Rsync/RsyncEngine.txt
+++ b/Engine/Build/Rsync/RsyncEngine.txt
@@ -155,6 +155,7 @@ protect /Saved/
 + /Source/Programs/Shared/
 + /Source/Programs/Shared/EpicGames.Core/
 + /Source/Programs/Shared/EpicGames.Build/
++ /Source/Programs/Shared/EpicGames.UHT/
 - /Source/Programs/Shared/*/
 + /Source/Programs/UnrealPak/
 + /Source/Programs/UnrealHeaderTool/
