macOS needs the dotnet 6 binaries copied over for remote iOS builds.  Since the remote mac could be either arm64 or intel, copy both over.

#rnx

diff --git a/Engine/Build/Rsync/RsyncEngine.txt b/Engine/Build/Rsync/RsyncEngine.txt
--- a/Engine/Build/Rsync/RsyncEngine.txt
+++ b/Engine/Build/Rsync/RsyncEngine.txt
@@ -62,6 +62,9 @@ protect /Binaries/TVOS/
 # Binaries/ThirdParty, just DotNet and VisualStudio required to run UnrealBuildTool, and libtbb used by UHT
 + /Binaries/ThirdParty/
 + /Binaries/ThirdParty/DotNet/
++ /Binaries/ThirdParty/DotNet/6.0.200/
++ /Binaries/ThirdParty/DotNet/6.0.200/mac**/
++ /Binaries/ThirdParty/DotNet/6.0.200/mac**/**
 - /Binaries/ThirdParty/DotNet/**
 + /Binaries/ThirdParty/VisualStudio/
 + /Binaries/ThirdParty/VisualStudio/**
