Rsync over the FacebookSDK xcframework files for Remote Mac builds.  Also, include the new xcframework files in the Install Build.

#rnx

diff --git a/Engine/Build/InstalledEngineFilters.xml b/Engine/Build/InstalledEngineFilters.xml
--- a/Engine/Build/InstalledEngineFilters.xml
+++ b/Engine/Build/InstalledEngineFilters.xml
@@ -542,6 +542,7 @@
 		Engine/Build/IOS/...
 		Engine/Source/ThirdParty/IOS/SoundSwitch/...
 		Engine/Source/ThirdParty/Eigen/...
+		Engine/Source/ThirdParty/Facebook/IOS/...
 		Engine/Source/ThirdParty/libJPG/...
 		Engine/Source/ThirdParty/mtlpp/...
 		Engine/Source/ThirdParty/libSampleRate/...
diff --git a/Engine/Build/Rsync/RsyncEngine.txt b/Engine/Build/Rsync/RsyncEngine.txt
--- a/Engine/Build/Rsync/RsyncEngine.txt
+++ b/Engine/Build/Rsync/RsyncEngine.txt
@@ -191,6 +191,8 @@ protect /Source/Programs/Shared/EpicGames.*/***
 + /Source/ThirdParty/mimalloc/src/**
 + /Source/ThirdParty/**/IOS/**.framework/**
 + /Source/ThirdParty/**/TVOS/**.framework/**
++ /Source/ThirdParty/**/IOS/**.xcframework/**
++ /Source/ThirdParty/**/TVOS/**.xcframework/**
 - /Source/ThirdParty/**
 
 # Source: Include everything else
