Fix prepare for debug exceptions caused by rsync

diff --git a/Engine/Build/Rsync/RsyncEngine.txt b/Engine/Build/Rsync/RsyncEngine.txt
--- a/Engine/Build/Rsync/RsyncEngine.txt
+++ b/Engine/Build/Rsync/RsyncEngine.txt
@@ -48,9 +48,12 @@
 + /Binaries/DotNET/UnrealBuildTool/
 + /Binaries/DotNET/UnrealBuildTool/**
 
-# Binaries/Mac: Include DsymExporter, ignore everything else
+# Binaries/Mac: Include DsymExporter and EpicWebHelper, ignore everything else
 + /Binaries/Mac/
-+ /Binaries/Mac/DsymExporter
++ /Binaries/Mac/DsymExporter.app
++ /Binaries/Mac/DsymExporter.app/***
++ /Binaries/Mac/EpicWebHelper.app/
++ /Binaries/Mac/EpicWebHelper.app/***
 protect /Binaries/Mac/**
 
 # Binaries/IOS
@@ -74,6 +77,17 @@ protect /Binaries/TVOS/
 + /Binaries/ThirdParty/Intel/TBB/Mac/
 + /Binaries/ThirdParty/Intel/TBB/Mac/libtbb.dylib
 + /Binaries/ThirdParty/Intel/TBB/Mac/libtbbmalloc.dylib
++ /Binaries/ThirdParty/Python3/
++ /Binaries/ThirdParty/Python3/Mac/
++ /Binaries/ThirdParty/Python3/Mac/lib/
++ /Binaries/ThirdParty/Python3/Mac/lib/libpython3.9.dylib
++ /Binaries/ThirdParty/svn/
++ /Binaries/ThirdParty/svn/Mac/
++ /Binaries/ThirdParty/svn/Mac/bin/
++ /Binaries/ThirdParty/svn/Mac/bin/**
++ /Binaries/ThirdParty/CEF3/
++ /Binaries/ThirdParty/CEF3/Mac/
++ /Binaries/ThirdParty/CEF3/Mac/***
 
 # Build
 + /Build/
@@ -113,14 +127,22 @@ protect /Intermediate/**
 
 # Plugins
 + /Plugins/
++ /Plugins/**
 + /Plugins/**/Binaries/Mac/
 + /Plugins/**/Binaries/IOS/
 + /Plugins/**/Binaries/TVOS/
 - /Plugins/**/Binaries/*/
 - /Plugins/**/Content/
 - /Plugins/**/Resources/
++ /Plugins/Enterprise/
++ /Plugins/Enterprise/MDLImporter/
++ /Plugins/Enterprise/MDLImporter/Binaries/
++ /Plugins/Enterprise/MDLImporter/Binaries/ThirdParty/
++ /Plugins/Enterprise/MDLImporter/Binaries/ThirdParty/MDL/
++ /Plugins/Enterprise/MDLImporter/Binaries/ThirdParty/MDL/Mac/
++ /Plugins/Enterprise/MDLImporter/Binaries/ThirdParty/MDL/Mac/***
 protect /Plugins/**/Intermediate/
-+ /Plugins/**
+
 
 # Programs
 + /Programs/
