Fix Window remote mac builds with incorrect rsync script. DotNet binaries are copied over in RsyncEngine.txt

#rnx

diff --git a/Engine/Build/Rsync/RsyncEngineScripts.txt b/Engine/Build/Rsync/RsyncEngineScripts.txt
--- a/Engine/Build/Rsync/RsyncEngineScripts.txt
+++ b/Engine/Build/Rsync/RsyncEngineScripts.txt
@@ -1,8 +1,6 @@
 # This file contains a list of shell scripts to be synced to the remote. These files will be locally converted to use Unix line endings, and will be
 # marked as executable on the remote. Wildcards are not supported.
 
-Binaries/ThirdParty/DotNet/6.0.302/mac-x64/
-Binaries/ThirdParty/DotNet/6.0.302/mac-armx64/
 Build/BatchFiles/RunDotnet.sh
 Build/BatchFiles/RunDotnetMSBuild.sh
 Build/BatchFiles/Mac/Build.sh
