- Added a third party .cpp file to get copied to remote macs

diff --git a/Engine/Build/Rsync/RsyncEngine.txt b/Engine/Build/Rsync/RsyncEngine.txt
--- a/Engine/Build/Rsync/RsyncEngine.txt
+++ b/Engine/Build/Rsync/RsyncEngine.txt
@@ -211,6 +211,7 @@ protect /Source/Programs/Shared/EpicGames.*/***
 + /Source/ThirdParty/libSampleRate/**.cpp
 + /Source/ThirdParty/Eigen/**
 + /Source/ThirdParty/mimalloc/src/**
++ /Source/ThirdParty/msdfgen/**.cpp
 + /Source/ThirdParty/**/IOS/**.framework/**
 + /Source/ThirdParty/**/TVOS/**.framework/**
 + /Source/ThirdParty/**/IOS/**.xcframework/**
