- Fixed remote builds of IOS from Windows

diff --git a/Engine/Build/Rsync/RsyncEngine.txt b/Engine/Build/Rsync/RsyncEngine.txt
--- a/Engine/Build/Rsync/RsyncEngine.txt
+++ b/Engine/Build/Rsync/RsyncEngine.txt
@@ -211,6 +211,7 @@ protect /Source/Programs/Shared/EpicGames.*/***
 + /Source/ThirdParty/Eigen/**
 + /Source/ThirdParty/mimalloc/src/**
 + /Source/ThirdParty/msdfgen/**.cpp
++ /Source/ThirdParty/nanosvg/**.cpp
 + /Source/ThirdParty/**/IOS/**.framework/**
 + /Source/ThirdParty/**/TVOS/**.framework/**
 + /Source/ThirdParty/**/IOS/**.xcframework/**
