[VFX Reference Update] Update remote sync file list for python 3.11 for mac builds

diff --git a/Engine/Build/Rsync/RsyncEngine.txt b/Engine/Build/Rsync/RsyncEngine.txt
--- a/Engine/Build/Rsync/RsyncEngine.txt
+++ b/Engine/Build/Rsync/RsyncEngine.txt
@@ -79,7 +79,7 @@ protect /Binaries/TVOS/
 + /Binaries/ThirdParty/Python3/
 + /Binaries/ThirdParty/Python3/Mac/
 + /Binaries/ThirdParty/Python3/Mac/lib/
-+ /Binaries/ThirdParty/Python3/Mac/lib/libpython3.9.dylib
++ /Binaries/ThirdParty/Python3/Mac/lib/libpython3.11.dylib
 + /Binaries/ThirdParty/svn/
 + /Binaries/ThirdParty/svn/Mac/
 + /Binaries/ThirdParty/svn/Mac/bin/
