Include dotnet 8 for remote iOS mac builds

#rnx

diff --git a/Engine/Build/Rsync/RsyncEngine.txt b/Engine/Build/Rsync/RsyncEngine.txt
--- a/Engine/Build/Rsync/RsyncEngine.txt
+++ b/Engine/Build/Rsync/RsyncEngine.txt
@@ -67,6 +67,9 @@ protect /Binaries/TVOS/
 + /Binaries/ThirdParty/DotNet/6.0.302/
 + /Binaries/ThirdParty/DotNet/6.0.302/mac**/
 + /Binaries/ThirdParty/DotNet/6.0.302/mac**/**
++ /Binaries/ThirdParty/DotNet/8.0.300/
++ /Binaries/ThirdParty/DotNet/8.0.300/mac**/
++ /Binaries/ThirdParty/DotNet/8.0.300/mac**/**
 - /Binaries/ThirdParty/DotNet/**
 + /Binaries/ThirdParty/VisualStudio/
 + /Binaries/ThirdParty/VisualStudio/**
