Add new ThirdParty lib to rsync script for remote iOS builds

#rnx

diff --git a/Engine/Build/Rsync/RsyncEngine.txt b/Engine/Build/Rsync/RsyncEngine.txt
--- a/Engine/Build/Rsync/RsyncEngine.txt
+++ b/Engine/Build/Rsync/RsyncEngine.txt
@@ -209,9 +209,11 @@ protect /Source/Programs/Shared/EpicGames.*/***
 + /Source/ThirdParty/**.dylib
 + /Source/ThirdParty/**.plist
 + /Source/ThirdParty/**.embeddedframework.zip
++ /Source/ThirdParty/**.xcframework.zip
 + /Source/ThirdParty/libJPG/*.cpp
 + /Source/ThirdParty/libSampleRate/**.cpp
 + /Source/ThirdParty/Eigen/**
++ /Source/ThirdParty/GuidelinesSupportLibrary/**
 + /Source/ThirdParty/mimalloc/src/**
 + /Source/ThirdParty/msdfgen/**.cpp
 + /Source/ThirdParty/nanosvg/**.cpp
