[Mac] Fix build error while using the Remote Mac Build feature.

A few weeks ago iOS platform code was updated to use mimalloc but there is a bug in the existing rsync rules we use where the mimalloc source files are not sent over to the remote mac.

The issue started happening now because of the recent change to start using mimalloc, but the rules have been incorrect for a while now.

Also, similar problem. The remote mac now also needs the MetalCPP source files but there was no rule to allow these to be sent over. This CL adds that missing rule.

#Synced-CL 51278809

diff --git a/Engine/Build/Rsync/RsyncEngine.txt b/Engine/Build/Rsync/RsyncEngine.txt
--- a/Engine/Build/Rsync/RsyncEngine.txt
+++ b/Engine/Build/Rsync/RsyncEngine.txt
@@ -206,13 +206,14 @@ protect /Source/Programs/Shared/EpicGames.*/***
 + /Source/ThirdParty/libSampleRate/**.cpp
 + /Source/ThirdParty/Eigen/**
 + /Source/ThirdParty/GuidelinesSupportLibrary/**
-+ /Source/ThirdParty/mimalloc/src/**
++ /Source/ThirdParty/mimalloc/**
 + /Source/ThirdParty/msdfgen/**.cpp
 + /Source/ThirdParty/nanosvg/**.cpp
 + /Source/ThirdParty/**/IOS/**.framework/**
 + /Source/ThirdParty/**/TVOS/**.framework/**
 + /Source/ThirdParty/**/IOS/**.xcframework/**
 + /Source/ThirdParty/**/TVOS/**.xcframework/**
++ /Source/ThirdParty/MetalCPP/**
 - /Source/ThirdParty/**
 - /Source/**/.vs/**
 
